[Motion Design] Fix Installed Build generation failure cause by AvalancheDisplayCluster

The module dependency chain is : AvalancheDisplayCluster -> AvalancheMedia -> Avalanche -> Text3D

AvalancheMedia, Avalanche and Text3D are all denied at the server level. Out of those 3, Text3D cannot be compiled in a Server target so this fails the Installed Build generation.

The fix is to set the TargetDenyList on the AvalancheDisplayCluster module. The Avalanche plugin's modules are all denied on Server already.

diff --git a/Engine/Plugins/Experimental/AvalancheDisplayCluster/AvalancheDisplayCluster.uplugin b/Engine/Plugins/Experimental/AvalancheDisplayCluster/AvalancheDisplayCluster.uplugin
--- a/Engine/Plugins/Experimental/AvalancheDisplayCluster/AvalancheDisplayCluster.uplugin
+++ b/Engine/Plugins/Experimental/AvalancheDisplayCluster/AvalancheDisplayCluster.uplugin
@@ -27,7 +27,8 @@
 			"PlatformAllowList": [
 				"Win64",
 				"Linux"
-			]
+			],
+			"TargetDenyList":  [ "Server" ]
 		}
 	],
 	"Plugins": [
